Fix deprecated warning on IOS.

diff --git a/Engine/Plugins/Runtime/OptionalMobileFeaturesBPLibrary/Source/OptionalMobileFeaturesBPLibrary/Private/OptionalMobileFeaturesBPLibrary.cpp b/Engine/Plugins/Runtime/OptionalMobileFeaturesBPLibrary/Source/OptionalMobileFeaturesBPLibrary/Private/OptionalMobileFeaturesBPLibrary.cpp
--- a/Engine/Plugins/Runtime/OptionalMobileFeaturesBPLibrary/Source/OptionalMobileFeaturesBPLibrary/Private/OptionalMobileFeaturesBPLibrary.cpp
+++ b/Engine/Plugins/Runtime/OptionalMobileFeaturesBPLibrary/Source/OptionalMobileFeaturesBPLibrary/Private/OptionalMobileFeaturesBPLibrary.cpp
@@ -18,7 +18,7 @@ int UOptionalMobileFeaturesBPLibrary::GetVolumeState()
 	int scaledVolume = (baseVolume * 100) / 15;
 	return scaledVolume;
 #elif PLATFORM_IOS
-	return FIOSPlatformMisc::GetAudioVolume();
+	return FPlatformMisc::GetDeviceVolume();
 #else
 	return 0;
 #endif
