Updated modules to not dependency on itself.

diff --git a/Engine/Plugins/Runtime/LocationServicesBPLibrary/Source/LocationServicesBPLibrary/LocationServicesBPLibrary.Build.cs b/Engine/Plugins/Runtime/LocationServicesBPLibrary/Source/LocationServicesBPLibrary/LocationServicesBPLibrary.Build.cs
--- a/Engine/Plugins/Runtime/LocationServicesBPLibrary/Source/LocationServicesBPLibrary/LocationServicesBPLibrary.Build.cs
+++ b/Engine/Plugins/Runtime/LocationServicesBPLibrary/Source/LocationServicesBPLibrary/LocationServicesBPLibrary.Build.cs
@@ -22,8 +22,6 @@ namespace UnrealBuildTool.Rules
 					// ... add private dependencies that you statically link with here ...
 				}
 			);
-
-            PublicIncludePathModuleNames.Add("LocationServicesBPLibrary");
         }
 	}
 }
diff --git a/Engine/Plugins/Runtime/OptionalMobileFeaturesBPLibrary/Source/OptionalMobileFeaturesBPLibrary/OptionalMobileFeaturesBPLibrary.Build.cs b/Engine/Plugins/Runtime/OptionalMobileFeaturesBPLibrary/Source/OptionalMobileFeaturesBPLibrary/OptionalMobileFeaturesBPLibrary.Build.cs
--- a/Engine/Plugins/Runtime/OptionalMobileFeaturesBPLibrary/Source/OptionalMobileFeaturesBPLibrary/OptionalMobileFeaturesBPLibrary.Build.cs
+++ b/Engine/Plugins/Runtime/OptionalMobileFeaturesBPLibrary/Source/OptionalMobileFeaturesBPLibrary/OptionalMobileFeaturesBPLibrary.Build.cs
@@ -23,8 +23,6 @@ namespace UnrealBuildTool.Rules
 					// ... add private dependencies that you statically link with here ...
 				}
 				);
-
-            PublicIncludePathModuleNames.Add("OptionalMobileFeaturesBPLibrary");
         }
 	}
 }
