Add a ShortName to modules that generated paths over the 200 chars limit and a few modules that were realy close (198+)

diff --git a/Engine/Plugins/Animation/AnimationLocomotionLibrary/Source/Editor/AnimationLocomotionLibraryEditor.Build.cs b/Engine/Plugins/Animation/AnimationLocomotionLibrary/Source/Editor/AnimationLocomotionLibraryEditor.Build.cs
--- a/Engine/Plugins/Animation/AnimationLocomotionLibrary/Source/Editor/AnimationLocomotionLibraryEditor.Build.cs
+++ b/Engine/Plugins/Animation/AnimationLocomotionLibrary/Source/Editor/AnimationLocomotionLibraryEditor.Build.cs
@@ -26,6 +26,8 @@ namespace UnrealBuildTool.Rules
 			new string[]
 			{
 			});
+
+			ShortName = "AnimLLEd";
 		}
 	}
 }
diff --git a/Engine/Plugins/Animation/AnimationLocomotionLibrary/Source/Runtime/AnimationLocomotionLibraryRuntime.Build.cs b/Engine/Plugins/Animation/AnimationLocomotionLibrary/Source/Runtime/AnimationLocomotionLibraryRuntime.Build.cs
--- a/Engine/Plugins/Animation/AnimationLocomotionLibrary/Source/Runtime/AnimationLocomotionLibraryRuntime.Build.cs
+++ b/Engine/Plugins/Animation/AnimationLocomotionLibrary/Source/Runtime/AnimationLocomotionLibraryRuntime.Build.cs
@@ -20,6 +20,8 @@ namespace UnrealBuildTool.Rules
 			new string[]
 			{
             });
+
+			ShortName = "AnimLLRun";
 		}
 	}
 }
diff --git a/Engine/Plugins/Developer/Concert/ConcertUI/ConcertClientSharedSlate/Source/ConcertClientSharedSlate/ConcertClientSharedSlate.Build.cs b/Engine/Plugins/Developer/Concert/ConcertUI/ConcertClientSharedSlate/Source/ConcertClientSharedSlate/ConcertClientSharedSlate.Build.cs
--- a/Engine/Plugins/Developer/Concert/ConcertUI/ConcertClientSharedSlate/Source/ConcertClientSharedSlate/ConcertClientSharedSlate.Build.cs
+++ b/Engine/Plugins/Developer/Concert/ConcertUI/ConcertClientSharedSlate/Source/ConcertClientSharedSlate/ConcertClientSharedSlate.Build.cs
@@ -18,6 +18,8 @@ namespace UnrealBuildTool.Rules
 					"ConcertSyncClient",
 				}
 			);
+
+			ShortName = "ConClShrSlt";
 		}
 	}
 }
diff --git a/Engine/Plugins/Experimental/Enterprise/DatasmithInterchange/Source/DatasmithInterchangeScripting/DatasmithInterchangeScripting.Build.cs b/Engine/Plugins/Experimental/Enterprise/DatasmithInterchange/Source/DatasmithInterchangeScripting/DatasmithInterchangeScripting.Build.cs
--- a/Engine/Plugins/Experimental/Enterprise/DatasmithInterchange/Source/DatasmithInterchangeScripting/DatasmithInterchangeScripting.Build.cs
+++ b/Engine/Plugins/Experimental/Enterprise/DatasmithInterchange/Source/DatasmithInterchangeScripting/DatasmithInterchangeScripting.Build.cs
@@ -29,6 +29,8 @@ namespace UnrealBuildTool.Rules
 					"InterchangeCore",
 				}
 			);
+
+			ShortName = "DSInterScript";
 		}
 	}
 }
\ No newline at end of file
diff --git a/Engine/Plugins/Experimental/VirtualProduction/LedWallCalibration/Source/LedWallCalibrationEditor/LedWallCalibrationEditor.Build.cs b/Engine/Plugins/Experimental/VirtualProduction/LedWallCalibration/Source/LedWallCalibrationEditor/LedWallCalibrationEditor.Build.cs
--- a/Engine/Plugins/Experimental/VirtualProduction/LedWallCalibration/Source/LedWallCalibrationEditor/LedWallCalibrationEditor.Build.cs
+++ b/Engine/Plugins/Experimental/VirtualProduction/LedWallCalibration/Source/LedWallCalibrationEditor/LedWallCalibrationEditor.Build.cs
@@ -33,5 +33,7 @@ public class LedWallCalibrationEditor : ModuleRules
 				"UnrealEd",
 			}
 		);
+
+		ShortName = "LedWCalEd";
 	}
 }
diff --git a/Engine/Plugins/MovieScene/CustomizableSequencerTracks/Source/CustomizableSequencerTracksEditor/CustomizableSequencerTracksEditor.Build.cs b/Engine/Plugins/MovieScene/CustomizableSequencerTracks/Source/CustomizableSequencerTracksEditor/CustomizableSequencerTracksEditor.Build.cs
--- a/Engine/Plugins/MovieScene/CustomizableSequencerTracks/Source/CustomizableSequencerTracksEditor/CustomizableSequencerTracksEditor.Build.cs
+++ b/Engine/Plugins/MovieScene/CustomizableSequencerTracks/Source/CustomizableSequencerTracksEditor/CustomizableSequencerTracksEditor.Build.cs
@@ -23,5 +23,7 @@ public class CustomizableSequencerTracksEditor : ModuleRules
 				"UnrealEd",
 			}
 		);
+
+		ShortName = "CustomSeqTrEd";
 	}
 }
diff --git a/Engine/Plugins/Runtime/OptionalMobileFeaturesBPLibrary/Source/OptionalMobileFeaturesBPLibrary/OptionalMobileFeaturesBPLibrary.Build.cs b/Engine/Plugins/Runtime/OptionalMobileFeaturesBPLibrary/Source/OptionalMobileFeaturesBPLibrary/OptionalMobileFeaturesBPLibrary.Build.cs
--- a/Engine/Plugins/Runtime/OptionalMobileFeaturesBPLibrary/Source/OptionalMobileFeaturesBPLibrary/OptionalMobileFeaturesBPLibrary.Build.cs
+++ b/Engine/Plugins/Runtime/OptionalMobileFeaturesBPLibrary/Source/OptionalMobileFeaturesBPLibrary/OptionalMobileFeaturesBPLibrary.Build.cs
@@ -23,6 +23,9 @@ namespace UnrealBuildTool.Rules
 					// ... add private dependencies that you statically link with here ...
 				}
 				);
-        }
+
+			ShortName = "OptMobFeatBPLib";
+
+		}
 	}
 }
diff --git a/Engine/Plugins/VirtualProduction/CameraCalibrationCore/Source/CameraCalibrationCoreMovieScene/CameraCalibrationCoreMovieScene.build.cs b/Engine/Plugins/VirtualProduction/CameraCalibrationCore/Source/CameraCalibrationCoreMovieScene/CameraCalibrationCoreMovieScene.build.cs
--- a/Engine/Plugins/VirtualProduction/CameraCalibrationCore/Source/CameraCalibrationCoreMovieScene/CameraCalibrationCoreMovieScene.build.cs
+++ b/Engine/Plugins/VirtualProduction/CameraCalibrationCore/Source/CameraCalibrationCoreMovieScene/CameraCalibrationCoreMovieScene.build.cs
@@ -36,5 +36,7 @@ public class CameraCalibrationCoreMovieScene : ModuleRules
  				"MovieScene",
 			}
 		);
+
+		ShortName = "CamCalCMovSC";
 	}
 }
diff --git a/Engine/Plugins/VirtualProduction/RemoteControlProtocolMIDI/Source/RemoteControlProtocolMIDIEditor/RemoteControlProtocolMIDIEditor.Build.cs b/Engine/Plugins/VirtualProduction/RemoteControlProtocolMIDI/Source/RemoteControlProtocolMIDIEditor/RemoteControlProtocolMIDIEditor.Build.cs
--- a/Engine/Plugins/VirtualProduction/RemoteControlProtocolMIDI/Source/RemoteControlProtocolMIDIEditor/RemoteControlProtocolMIDIEditor.Build.cs
+++ b/Engine/Plugins/VirtualProduction/RemoteControlProtocolMIDI/Source/RemoteControlProtocolMIDIEditor/RemoteControlProtocolMIDIEditor.Build.cs
@@ -30,5 +30,7 @@ public class RemoteControlProtocolMIDIEditor : ModuleRules
 				"SlateCore",
 			}
 		);
+
+		ShortName = "RCPMidiEd";
 	}
 }
